LensDistortion converted over to renderpasses.

diff --git a/Engine/Plugins/Compositing/LensDistortion/Source/LensDistortion/Private/LensDistortionRendering.cpp b/Engine/Plugins/Compositing/LensDistortion/Source/LensDistortion/Private/LensDistortionRendering.cpp
--- a/Engine/Plugins/Compositing/LensDistortion/Source/LensDistortion/Private/LensDistortionRendering.cpp
+++ b/Engine/Plugins/Compositing/LensDistortion/Source/LensDistortion/Private/LensDistortionRendering.cpp
@@ -177,56 +177,47 @@ static void DrawUVDisplacementToRenderTarget_RenderThread(
 	SCOPED_DRAW_EVENT(RHICmdList, DrawUVDisplacementToRenderTarget_RenderThread);
 #endif
 
-	// Set render target.
-	SetRenderTarget(
-		RHICmdList,
-		OutTextureRenderTargetResource->GetRenderTargetTexture(),
-		FTextureRHIRef(),
-		ESimpleRenderTargetMode::EUninitializedColorAndDepth,
-		FExclusiveDepthStencil::DepthNop_StencilNop);
-
-	FIntPoint DisplacementMapResolution(OutTextureRenderTargetResource->GetSizeX(), OutTextureRenderTargetResource->GetSizeY());
-
-	// Update viewport.
-	RHICmdList.SetViewport(
-		0, 0, 0.f,
-		DisplacementMapResolution.X, DisplacementMapResolution.Y, 1.f);
-	
-	// Get shaders.
-	TShaderMap<FGlobalShaderType>* GlobalShaderMap = GetGlobalShaderMap(FeatureLevel);
-	TShaderMapRef< FLensDistortionUVGenerationVS > VertexShader(GlobalShaderMap);
-	TShaderMapRef< FLensDistortionUVGenerationPS > PixelShader(GlobalShaderMap);
-
-	// Set the graphic pipeline state.
-	FGraphicsPipelineStateInitializer GraphicsPSOInit;
-	RHICmdList.ApplyCachedRenderTargets(GraphicsPSOInit);
-	GraphicsPSOInit.DepthStencilState = TStaticDepthStencilState<false, CF_Always>::GetRHI();
-	GraphicsPSOInit.BlendState = TStaticBlendState<>::GetRHI();
-	GraphicsPSOInit.RasterizerState = TStaticRasterizerState<>::GetRHI();
-	GraphicsPSOInit.PrimitiveType = PT_TriangleList;
-	GraphicsPSOInit.BoundShaderState.VertexDeclarationRHI = GetVertexDeclarationFVector4();
-	GraphicsPSOInit.BoundShaderState.VertexShaderRHI = GETSAFERHISHADER_VERTEX(*VertexShader);
-	GraphicsPSOInit.BoundShaderState.PixelShaderRHI = GETSAFERHISHADER_PIXEL(*PixelShader);
-	SetGraphicsPipelineState(RHICmdList, GraphicsPSOInit);
-
-	// Update viewport.
-	RHICmdList.SetViewport(
-		0, 0, 0.f,
-		OutTextureRenderTargetResource->GetSizeX(), OutTextureRenderTargetResource->GetSizeY(), 1.f);
-
-	// Update shader uniform parameters.
-	VertexShader->SetParameters(RHICmdList, VertexShader->GetVertexShader(), CompiledCameraModel, DisplacementMapResolution);
-	PixelShader->SetParameters(RHICmdList, PixelShader->GetPixelShader(), CompiledCameraModel, DisplacementMapResolution);
-
-	// Draw grid.
-	uint32 PrimitiveCount = kGridSubdivisionX * kGridSubdivisionY * 2;
-	RHICmdList.DrawPrimitive(0, PrimitiveCount, 1);
-
-	// Resolve render target.
-	RHICmdList.CopyToResolveTarget(
-		OutTextureRenderTargetResource->GetRenderTargetTexture(),
-		OutTextureRenderTargetResource->TextureRHI,
-		FResolveParams());
+	FRHIRenderPassInfo RPInfo(OutTextureRenderTargetResource->GetRenderTargetTexture(), ERenderTargetActions::DontLoad_Store, OutTextureRenderTargetResource->TextureRHI);
+	RHICmdList.BeginRenderPass(RPInfo, TEXT("DrawUVDisplacement"));
+	{
+		FIntPoint DisplacementMapResolution(OutTextureRenderTargetResource->GetSizeX(), OutTextureRenderTargetResource->GetSizeY());
+
+		// Update viewport.
+		RHICmdList.SetViewport(
+			0, 0, 0.f,
+			DisplacementMapResolution.X, DisplacementMapResolution.Y, 1.f);
+
+		// Get shaders.
+		TShaderMap<FGlobalShaderType>* GlobalShaderMap = GetGlobalShaderMap(FeatureLevel);
+		TShaderMapRef< FLensDistortionUVGenerationVS > VertexShader(GlobalShaderMap);
+		TShaderMapRef< FLensDistortionUVGenerationPS > PixelShader(GlobalShaderMap);
+
+		// Set the graphic pipeline state.
+		FGraphicsPipelineStateInitializer GraphicsPSOInit;
+		RHICmdList.ApplyCachedRenderTargets(GraphicsPSOInit);
+		GraphicsPSOInit.DepthStencilState = TStaticDepthStencilState<false, CF_Always>::GetRHI();
+		GraphicsPSOInit.BlendState = TStaticBlendState<>::GetRHI();
+		GraphicsPSOInit.RasterizerState = TStaticRasterizerState<>::GetRHI();
+		GraphicsPSOInit.PrimitiveType = PT_TriangleList;
+		GraphicsPSOInit.BoundShaderState.VertexDeclarationRHI = GetVertexDeclarationFVector4();
+		GraphicsPSOInit.BoundShaderState.VertexShaderRHI = GETSAFERHISHADER_VERTEX(*VertexShader);
+		GraphicsPSOInit.BoundShaderState.PixelShaderRHI = GETSAFERHISHADER_PIXEL(*PixelShader);
+		SetGraphicsPipelineState(RHICmdList, GraphicsPSOInit);
+
+		// Update viewport.
+		RHICmdList.SetViewport(
+			0, 0, 0.f,
+			OutTextureRenderTargetResource->GetSizeX(), OutTextureRenderTargetResource->GetSizeY(), 1.f);
+
+		// Update shader uniform parameters.
+		VertexShader->SetParameters(RHICmdList, VertexShader->GetVertexShader(), CompiledCameraModel, DisplacementMapResolution);
+		PixelShader->SetParameters(RHICmdList, PixelShader->GetPixelShader(), CompiledCameraModel, DisplacementMapResolution);
+
+		// Draw grid.
+		uint32 PrimitiveCount = kGridSubdivisionX * kGridSubdivisionY * 2;
+		RHICmdList.DrawPrimitive(0, PrimitiveCount, 1);
+	}
+	RHICmdList.EndRenderPass();
 }
 
 
diff --git a/Engine/Plugins/Compositing/OpenCVLensDistortion/Source/OpenCVLensDistortion/Private/LensDistortionDisplacementMapRendering.cpp b/Engine/Plugins/Compositing/OpenCVLensDistortion/Source/OpenCVLensDistortion/Private/LensDistortionDisplacementMapRendering.cpp
--- a/Engine/Plugins/Compositing/OpenCVLensDistortion/Source/OpenCVLensDistortion/Private/LensDistortionDisplacementMapRendering.cpp
+++ b/Engine/Plugins/Compositing/OpenCVLensDistortion/Source/OpenCVLensDistortion/Private/LensDistortionDisplacementMapRendering.cpp
@@ -124,39 +124,40 @@ static void DrawUVDisplacementToRenderTargetFromPreComputedDisplacementMap_Rende
 #endif
 
 	// Set render target.
-	SetRenderTarget(RHICmdList, OutTextureRenderTargetResource->GetRenderTargetTexture(), FTextureRHIRef(),	ESimpleRenderTargetMode::EClearColorAndDepth, FExclusiveDepthStencil::DepthNop_StencilNop);
-
-	// Update viewport.
-	const FIntPoint DisplacementMapResolution(OutTextureRenderTargetResource->GetSizeX(), OutTextureRenderTargetResource->GetSizeY());
-	RHICmdList.SetViewport(0, 0, 0.f, DisplacementMapResolution.X, DisplacementMapResolution.Y, 1.f);
-	
-	// Get shaders.
-	TShaderMap<FGlobalShaderType>* GlobalShaderMap = GetGlobalShaderMap(FeatureLevel);
-	TShaderMapRef< FLensDistortionDisplacementMapGenerationVS > VertexShader(GlobalShaderMap);
-	TShaderMapRef< FLensDistortionDisplacementMapGenerationPS > PixelShader(GlobalShaderMap);
-
-	// Set the graphic pipeline state.
-	FGraphicsPipelineStateInitializer GraphicsPSOInit;
-	RHICmdList.ApplyCachedRenderTargets(GraphicsPSOInit);
-	GraphicsPSOInit.DepthStencilState = TStaticDepthStencilState<false, CF_Always>::GetRHI();
-	GraphicsPSOInit.BlendState = TStaticBlendState<>::GetRHI();
-	GraphicsPSOInit.RasterizerState = TStaticRasterizerState<>::GetRHI();
-	GraphicsPSOInit.PrimitiveType = PT_TriangleList;
-	GraphicsPSOInit.BoundShaderState.VertexDeclarationRHI = GetVertexDeclarationFVector4();
-	GraphicsPSOInit.BoundShaderState.VertexShaderRHI = GETSAFERHISHADER_VERTEX(*VertexShader);
-	GraphicsPSOInit.BoundShaderState.PixelShaderRHI = GETSAFERHISHADER_PIXEL(*PixelShader);
-	SetGraphicsPipelineState(RHICmdList, GraphicsPSOInit);
-
-	// Update shader uniform parameters.
-	VertexShader->SetParameters(RHICmdList, VertexShader->GetVertexShader(), PreComputedDisplacementMap, DisplacementMapResolution);
-	PixelShader->SetParameters(RHICmdList, PixelShader->GetPixelShader(), PreComputedDisplacementMap, DisplacementMapResolution);
-
-	// Draw grid.
-	const uint32 PrimitiveCount = kGridSubdivisionX * kGridSubdivisionY * 2;
-	RHICmdList.DrawPrimitive(0, PrimitiveCount, 1);
-
-	// Resolve render target.
-	RHICmdList.CopyToResolveTarget(OutTextureRenderTargetResource->GetRenderTargetTexture(), OutTextureRenderTargetResource->TextureRHI, FResolveParams());
+	FRHIRenderPassInfo RPInfo(OutTextureRenderTargetResource->GetRenderTargetTexture(), ERenderTargetActions::Clear_Store, OutTextureRenderTargetResource->TextureRHI);
+	RHICmdList.BeginRenderPass(RPInfo, TEXT("DrawUVDisplacementFromPrecomputedDisplacementMap"));
+	{
+
+		// Update viewport.
+		const FIntPoint DisplacementMapResolution(OutTextureRenderTargetResource->GetSizeX(), OutTextureRenderTargetResource->GetSizeY());
+		RHICmdList.SetViewport(0, 0, 0.f, DisplacementMapResolution.X, DisplacementMapResolution.Y, 1.f);
+
+		// Get shaders.
+		TShaderMap<FGlobalShaderType>* GlobalShaderMap = GetGlobalShaderMap(FeatureLevel);
+		TShaderMapRef< FLensDistortionDisplacementMapGenerationVS > VertexShader(GlobalShaderMap);
+		TShaderMapRef< FLensDistortionDisplacementMapGenerationPS > PixelShader(GlobalShaderMap);
+
+		// Set the graphic pipeline state.
+		FGraphicsPipelineStateInitializer GraphicsPSOInit;
+		RHICmdList.ApplyCachedRenderTargets(GraphicsPSOInit);
+		GraphicsPSOInit.DepthStencilState = TStaticDepthStencilState<false, CF_Always>::GetRHI();
+		GraphicsPSOInit.BlendState = TStaticBlendState<>::GetRHI();
+		GraphicsPSOInit.RasterizerState = TStaticRasterizerState<>::GetRHI();
+		GraphicsPSOInit.PrimitiveType = PT_TriangleList;
+		GraphicsPSOInit.BoundShaderState.VertexDeclarationRHI = GetVertexDeclarationFVector4();
+		GraphicsPSOInit.BoundShaderState.VertexShaderRHI = GETSAFERHISHADER_VERTEX(*VertexShader);
+		GraphicsPSOInit.BoundShaderState.PixelShaderRHI = GETSAFERHISHADER_PIXEL(*PixelShader);
+		SetGraphicsPipelineState(RHICmdList, GraphicsPSOInit);
+
+		// Update shader uniform parameters.
+		VertexShader->SetParameters(RHICmdList, VertexShader->GetVertexShader(), PreComputedDisplacementMap, DisplacementMapResolution);
+		PixelShader->SetParameters(RHICmdList, PixelShader->GetPixelShader(), PreComputedDisplacementMap, DisplacementMapResolution);
+
+		// Draw grid.
+		const uint32 PrimitiveCount = kGridSubdivisionX * kGridSubdivisionY * 2;
+		RHICmdList.DrawPrimitive(0, PrimitiveCount, 1);
+	}
+	RHICmdList.EndRenderPass();
 #endif
 }
 
