Localization Duplicate: LensDistortionRendering
-Changed both keys so they are unique and give better context.

diff --git a/Engine/Plugins/Compositing/LensDistortion/Source/LensDistortion/Private/LensDistortionRendering.cpp b/Engine/Plugins/Compositing/LensDistortion/Source/LensDistortion/Private/LensDistortionRendering.cpp
--- a/Engine/Plugins/Compositing/LensDistortion/Source/LensDistortion/Private/LensDistortionRendering.cpp
+++ b/Engine/Plugins/Compositing/LensDistortion/Source/LensDistortion/Private/LensDistortionRendering.cpp
@@ -322,7 +322,7 @@ void FLensDistortionCameraModel::DrawUVDisplacementToRenderTarget(
 
 	if (!OutputRenderTarget)
 	{
-		FMessageLog("Blueprint").Warning(LOCTEXT("LensDistortionCameraModel_DrawUVDisplacementToRenderTarget", "DrawUVDisplacementToRenderTarget: Output render target is required."));
+		FMessageLog("Blueprint").Warning(LOCTEXT("LensDistortionCameraModel_OutputTargetRequired", "DrawUVDisplacementToRenderTarget: Output render target is required."));
 		return;
 	}
 
@@ -354,7 +354,7 @@ void FLensDistortionCameraModel::DrawUVDisplacementToRenderTarget(
 
 	if (FeatureLevel < ERHIFeatureLevel::SM5)
 	{
-		FMessageLog("Blueprint").Warning(LOCTEXT("LensDistortionCameraModel_DrawUVDisplacementToRenderTarget", "DrawUVDisplacementToRenderTarget: Requires RHIFeatureLevel::SM5 which is unavailable."));
+		FMessageLog("Blueprint").Warning(LOCTEXT("LensDistortionCameraModel_SM5Unavailable", "DrawUVDisplacementToRenderTarget: Requires RHIFeatureLevel::SM5 which is unavailable."));
 		return;
 	}
 
