Deprecate LensDistortion plugin

Remove LensDistortion as an enabled plugin in multiple plugins
Remove automated tests from EngineTest
Remove LensDistortion as a plugin dependency from composure
Remove now unused BP function from Composure's BP helper library that depends on a function from the LensDistortion BP function library

diff --git a/Engine/Plugins/Compositing/Composure/Composure.uplugin b/Engine/Plugins/Compositing/Composure/Composure.uplugin
--- a/Engine/Plugins/Compositing/Composure/Composure.uplugin
+++ b/Engine/Plugins/Compositing/Composure/Composure.uplugin
@@ -37,10 +37,6 @@
 			"Name": "CameraCalibrationCore",
 			"Enabled": true
 		},
-		{
-			"Name": "LensDistortion",
-			"Enabled": true
-		},
 		{
 			"Name": "BlueprintMaterialTextureNodes",
 			"Enabled": true
diff --git a/Engine/Plugins/Compositing/LensDistortion/LensDistortion.uplugin b/Engine/Plugins/Compositing/LensDistortion/LensDistortion.uplugin
--- a/Engine/Plugins/Compositing/LensDistortion/LensDistortion.uplugin
+++ b/Engine/Plugins/Compositing/LensDistortion/LensDistortion.uplugin
@@ -2,8 +2,8 @@
 	"FileVersion" : 3,
 	"Version" : 1,
 	"VersionName" : "1.0",
-	"FriendlyName" : "Lens Distortion",
-	"Description" : "Plugin to generate UV displacement for lens distortion/undistortion on the GPU from standard camera model.",
+	"FriendlyName" : "Lens Distortion (Deprecated)",
+	"Description" : "This plugin has been deprecated and will be removed in a future engine version. Please update your project to use the features of the CameraCalibration plugin instead.\nPlugin to generate UV displacement for lens distortion/undistortion on the GPU from standard camera model.",
 	"Category" : "Compositing",
 	"CreatedBy" : "Epic Games, Inc.",
 	"CreatedByURL" : "http://epicgames.com",
diff --git a/Engine/Plugins/Compositing/LensDistortion/Source/LensDistortion/Classes/LensDistortionAPI.h b/Engine/Plugins/Compositing/LensDistortion/Source/LensDistortion/Classes/LensDistortionAPI.h
--- a/Engine/Plugins/Compositing/LensDistortion/Source/LensDistortion/Classes/LensDistortionAPI.h
+++ b/Engine/Plugins/Compositing/LensDistortion/Source/LensDistortion/Classes/LensDistortionAPI.h
@@ -13,6 +13,7 @@
  *  |  0  F.Y C.Y |
  *  |  0   0   1  |
  */
+struct UE_DEPRECATED(5.0, "The LensDistortion plugin is deprecated. Please update your project to use the features of the CameraCalibration plugin.") FLensDistortionCameraModel;
 USTRUCT(BlueprintType)
 struct LENSDISTORTION_API FLensDistortionCameraModel
 {
diff --git a/Engine/Plugins/Compositing/LensDistortion/Source/LensDistortion/Classes/LensDistortionBlueprintLibrary.h b/Engine/Plugins/Compositing/LensDistortion/Source/LensDistortion/Classes/LensDistortionBlueprintLibrary.h
--- a/Engine/Plugins/Compositing/LensDistortion/Source/LensDistortion/Classes/LensDistortionBlueprintLibrary.h
+++ b/Engine/Plugins/Compositing/LensDistortion/Source/LensDistortion/Classes/LensDistortionBlueprintLibrary.h
@@ -9,14 +9,16 @@
 #include "LensDistortionBlueprintLibrary.generated.h"
 
 
-UCLASS(MinimalAPI, meta=(ScriptName="LensDistortionLibrary"))
+class UE_DEPRECATED(5.0, "The LensDistortion plugin is deprecated. Please update your project to use the features of the CameraCalibration plugin.") ULensDistortionBlueprintLibrary;
+UCLASS(MinimalAPI, meta=(DeprecationMessage = "The LensDistortion plugin is deprecated. Please update your project to use the features of the CameraCalibration plugin.", ScriptName="LensDistortionLibrary"))
 class ULensDistortionBlueprintLibrary : public UBlueprintFunctionLibrary
 {
 	GENERATED_UCLASS_BODY()
 
 
+	PRAGMA_DISABLE_DEPRECATION_WARNINGS
 	/** Returns the overscan factor required for the undistort rendering to avoid unrendered distorted pixels. */
-	UFUNCTION(BlueprintPure, Category = "Lens Distortion")
+	UFUNCTION(BlueprintPure, Category = "Lens Distortion", meta = (DeprecatedFunction, DeprecationMessage = "The LensDistortion plugin is deprecated. Please update your project to use the features of the CameraCalibration plugin."))
 	static void GetUndistortOverscanFactor(
 		const FLensDistortionCameraModel& CameraModel,
 		float DistortedHorizontalFOV,
@@ -33,7 +35,7 @@ class ULensDistortionBlueprintLibrary : public UBlueprintFunctionLibrary
 	 * @param OutputMultiply The multiplication factor applied on the displacement.
 	 * @param OutputAdd Value added to the multiplied displacement before storing into the output render target.
 	 */
-	UFUNCTION(BlueprintCallable, Category = "Lens Distortion", meta = (WorldContext = "WorldContextObject"))
+	UFUNCTION(BlueprintCallable, Category = "Lens Distortion", meta = (DeprecatedFunction, DeprecationMessage = "The LensDistortion plugin is deprecated. Please update your project to use the features of the CameraCalibration plugin.", WorldContext = "WorldContextObject"))
 	static void DrawUVDisplacementToRenderTarget(
 		const UObject* WorldContextObject,
 		const FLensDistortionCameraModel& CameraModel,
@@ -46,7 +48,7 @@ class ULensDistortionBlueprintLibrary : public UBlueprintFunctionLibrary
 		);
 
 	/* Returns true if A is equal to B (A == B) */
-	UFUNCTION(BlueprintPure, meta=(DisplayName = "Equal (LensDistortionCameraModel)", CompactNodeTitle = "==", Keywords = "== equal"), Category = "Lens Distortion")
+	UFUNCTION(BlueprintPure, meta=(DeprecatedFunction, DeprecationMessage = "The LensDistortion plugin is deprecated. Please update your project to use the features of the CameraCalibration plugin.", DisplayName = "Equal (LensDistortionCameraModel)", CompactNodeTitle = "==", Keywords = "== equal"), Category = "Lens Distortion")
 	static bool EqualEqual_CompareLensDistortionModels(
 		const FLensDistortionCameraModel& A,
 		const FLensDistortionCameraModel& B)
@@ -55,11 +57,13 @@ class ULensDistortionBlueprintLibrary : public UBlueprintFunctionLibrary
 	}
 
 	/* Returns true if A is not equal to B (A != B) */
-	UFUNCTION(BlueprintPure, meta = (DisplayName = "NotEqual (LensDistortionCameraModel)", CompactNodeTitle = "!=", Keywords = "!= not equal"), Category = "Lens Distortion")
+	UFUNCTION(BlueprintPure, meta = (DeprecatedFunction, DeprecationMessage = "The LensDistortion plugin is deprecated. Please update your project to use the features of the CameraCalibration plugin.", DisplayName = "NotEqual (LensDistortionCameraModel)", CompactNodeTitle = "!=", Keywords = "!= not equal"), Category = "Lens Distortion")
 	static bool NotEqual_CompareLensDistortionModels(
 		const FLensDistortionCameraModel& A,
 		const FLensDistortionCameraModel& B)
 	{
 		return A != B;
 	}
+	PRAGMA_ENABLE_DEPRECATION_WARNINGS
+
 };
diff --git a/Engine/Plugins/Compositing/LensDistortion/Source/LensDistortion/Private/LensDistortion.cpp b/Engine/Plugins/Compositing/LensDistortion/Source/LensDistortion/Private/LensDistortion.cpp
--- a/Engine/Plugins/Compositing/LensDistortion/Source/LensDistortion/Private/LensDistortion.cpp
+++ b/Engine/Plugins/Compositing/LensDistortion/Source/LensDistortion/Private/LensDistortion.cpp
@@ -7,7 +7,7 @@
 #include "Misc/Paths.h"
 #include "ShaderCore.h"
 
-
+PRAGMA_DISABLE_DEPRECATION_WARNINGS
 class FLensDistortion : public ILensDistortion
 {
 	/** IModuleInterface implementation */
@@ -30,3 +30,4 @@ void FLensDistortion::ShutdownModule()
 	// This function may be called during shutdown to clean up your module.  For modules that support dynamic reloading,
 	// we call this function before unloading the module.
 }
+PRAGMA_ENABLE_DEPRECATION_WARNINGS
diff --git a/Engine/Plugins/Compositing/LensDistortion/Source/LensDistortion/Private/LensDistortionBlueprintLibrary.cpp b/Engine/Plugins/Compositing/LensDistortion/Source/LensDistortion/Private/LensDistortionBlueprintLibrary.cpp
--- a/Engine/Plugins/Compositing/LensDistortion/Source/LensDistortion/Private/LensDistortionBlueprintLibrary.cpp
+++ b/Engine/Plugins/Compositing/LensDistortion/Source/LensDistortion/Private/LensDistortionBlueprintLibrary.cpp
@@ -3,7 +3,7 @@
 #include "LensDistortionBlueprintLibrary.h"
 
 
-
+PRAGMA_DISABLE_DEPRECATION_WARNINGS
 ULensDistortionBlueprintLibrary::ULensDistortionBlueprintLibrary(const FObjectInitializer& ObjectInitializer)
 	: Super(ObjectInitializer)
 { }
@@ -37,3 +37,4 @@ void ULensDistortionBlueprintLibrary::DrawUVDisplacementToRenderTarget(
 		UndistortOverscanFactor, OutputRenderTarget,
 		OutputMultiply, OutputAdd);
 }
+PRAGMA_ENABLE_DEPRECATION_WARNINGS
diff --git a/Engine/Plugins/Compositing/LensDistortion/Source/LensDistortion/Private/LensDistortionRendering.cpp b/Engine/Plugins/Compositing/LensDistortion/Source/LensDistortion/Private/LensDistortionRendering.cpp
--- a/Engine/Plugins/Compositing/LensDistortion/Source/LensDistortion/Private/LensDistortionRendering.cpp
+++ b/Engine/Plugins/Compositing/LensDistortion/Source/LensDistortion/Private/LensDistortionRendering.cpp
@@ -25,6 +25,7 @@ static const uint32 kGridSubdivisionY = 16;
  * Internal intermediary structure derived from FLensDistortionCameraModel by the game thread
  * to hand to the render thread.
  */
+PRAGMA_DISABLE_DEPRECATION_WARNINGS
 struct FCompiledCameraModel
 {
 	/** Orignal camera model that has generated this compiled model. */
@@ -370,5 +371,5 @@ void FLensDistortionCameraModel::DrawUVDisplacementToRenderTarget(
 		}
 	);
 }
-
+PRAGMA_ENABLE_DEPRECATION_WARNINGS
 #undef LOCTEXT_NAMESPACE
diff --git a/Engine/Plugins/Compositing/LensDistortion/Source/LensDistortion/Public/ILensDistortion.h b/Engine/Plugins/Compositing/LensDistortion/Source/LensDistortion/Public/ILensDistortion.h
--- a/Engine/Plugins/Compositing/LensDistortion/Source/LensDistortion/Public/ILensDistortion.h
+++ b/Engine/Plugins/Compositing/LensDistortion/Source/LensDistortion/Public/ILensDistortion.h
@@ -10,6 +10,7 @@
 /**
  * The public interface to this module
  */
+class UE_DEPRECATED(5.0, "The LensDistortion plugin is deprecated. Please update your project to use the features of the CameraCalibration plugin.") ILensDistortion;
 class ILensDistortion : public IModuleInterface
 {
 
@@ -21,6 +22,7 @@ public:
 	 *
 	 * @return Returns singleton instance, loading the module on demand if needed
 	 */
+	UE_DEPRECATED(5.0, "The LensDistortion plugin is deprecated. Please update your project to use the features of the CameraCalibration plugin.")
 	static inline ILensDistortion& Get()
 	{
 		return FModuleManager::LoadModuleChecked< ILensDistortion >( "LensDistortion" );
@@ -31,6 +33,7 @@ public:
 	 *
 	 * @return True if the module is loaded and ready to use
 	 */
+	UE_DEPRECATED(5.0, "The LensDistortion plugin is deprecated. Please update your project to use the features of the CameraCalibration plugin.")
 	static inline bool IsAvailable()
 	{
 		return FModuleManager::Get().IsModuleLoaded( "LensDistortion" );
